Make PanoramicCapture a developer module. Since it now depends on MessageLog (also a developer module), it fails to compile in shipping configurations, preventing installed engine from building correctly.

#rnx

diff --git a/Engine/Plugins/Experimental/PanoramicCapture/PanoramicCapture.uplugin b/Engine/Plugins/Experimental/PanoramicCapture/PanoramicCapture.uplugin
--- a/Engine/Plugins/Experimental/PanoramicCapture/PanoramicCapture.uplugin
+++ b/Engine/Plugins/Experimental/PanoramicCapture/PanoramicCapture.uplugin
@@ -17,7 +17,7 @@
 	"Modules": [
 		{
 			"Name": "PanoramicCapture",
-			"Type": "RuntimeNoCommandlet",
+			"Type": "Developer",
 			"LoadingPhase": "Default",
 			"WhitelistPlatforms": [
 				"Win64"
