Fixed localization gather warnings

#rnx

diff --git a/Engine/Plugins/Animation/LiveLink/Source/LiveLinkMovieScene/Private/LiveLinkMovieSceneModule.cpp b/Engine/Plugins/Animation/LiveLink/Source/LiveLinkMovieScene/Private/LiveLinkMovieSceneModule.cpp
--- a/Engine/Plugins/Animation/LiveLink/Source/LiveLinkMovieScene/Private/LiveLinkMovieSceneModule.cpp
+++ b/Engine/Plugins/Animation/LiveLink/Source/LiveLinkMovieScene/Private/LiveLinkMovieSceneModule.cpp
@@ -9,5 +9,3 @@ DEFINE_LOG_CATEGORY(LogLiveLinkMovieScene);
 
 
 IMPLEMENT_MODULE(FDefaultModuleImpl, LiveLinkMovieScene);
-
-#undef LOCTEXT_NAMESPACE
\ No newline at end of file
diff --git a/Engine/Plugins/Animation/LiveLink/Source/LiveLinkMovieScene/Private/MovieScene/MovieSceneLiveLinkSubSectionProperties.cpp b/Engine/Plugins/Animation/LiveLink/Source/LiveLinkMovieScene/Private/MovieScene/MovieSceneLiveLinkSubSectionProperties.cpp
--- a/Engine/Plugins/Animation/LiveLink/Source/LiveLinkMovieScene/Private/MovieScene/MovieSceneLiveLinkSubSectionProperties.cpp
+++ b/Engine/Plugins/Animation/LiveLink/Source/LiveLinkMovieScene/Private/MovieScene/MovieSceneLiveLinkSubSectionProperties.cpp
@@ -65,7 +65,7 @@ int32 UMovieSceneLiveLinkSubSectionProperties::CreateChannelProxy(int32 InChanne
 	{
 		if(UProperty* PropertyPtr = FindField<UProperty>(SubjectRole.GetDefaultObject()->GetFrameDataStruct(), Data.PropertyName))
 		{
-			const FText PropertyName = FText::Format(LOCTEXT("LinkLinkFormat", "{0}"), FText::FromName(Data.PropertyName));
+			const FText PropertyName = FText::FromName(Data.PropertyName);
  			if (PropertyPtr->IsA(UFloatProperty::StaticClass()))
 			{
 #if WITH_EDITOR
diff --git a/Engine/Plugins/Developer/Concert/MultiUserClient/Source/MultiUserClient/Private/Widgets/SConcertBrowser.cpp b/Engine/Plugins/Developer/Concert/MultiUserClient/Source/MultiUserClient/Private/Widgets/SConcertBrowser.cpp
--- a/Engine/Plugins/Developer/Concert/MultiUserClient/Source/MultiUserClient/Private/Widgets/SConcertBrowser.cpp
+++ b/Engine/Plugins/Developer/Concert/MultiUserClient/Source/MultiUserClient/Private/Widgets/SConcertBrowser.cpp
@@ -2284,7 +2284,7 @@ TSharedRef<SWidget> SConcertSessionBrowser::MakeButtonBar()
 			// Launch server.
 			+SUniformGridPanel::Slot(0, 0)
 			[
-				ConcertBrowserUtils::MakeIconButton(TEXT("FlatButton.Info"), FEditorFontGlyphs::Server, LOCTEXT("LaunchServerTooltip", "Launch a Multi-User server on this computer if one is not already running"),
+				ConcertBrowserUtils::MakeIconButton(TEXT("FlatButton.Info"), FEditorFontGlyphs::Server, LOCTEXT("LaunchServerTooltip", "Launch a Multi-User server on your computer unless one is already running"),
 					TAttribute<bool>(this, &SConcertSessionBrowser::IsLaunchServerButtonEnabled),
 					FOnClicked::CreateSP(this, &SConcertSessionBrowser::OnLaunchServerButtonClicked),
 					EVisibility::Visible,
diff --git a/Engine/Plugins/Experimental/ChaosEditor/Source/FractureEditor/Private/FractureToolCluster.cpp b/Engine/Plugins/Experimental/ChaosEditor/Source/FractureEditor/Private/FractureToolCluster.cpp
--- a/Engine/Plugins/Experimental/ChaosEditor/Source/FractureEditor/Private/FractureToolCluster.cpp
+++ b/Engine/Plugins/Experimental/ChaosEditor/Source/FractureEditor/Private/FractureToolCluster.cpp
@@ -51,15 +51,7 @@ FSlateIcon UFractureToolCluster::GetToolIcon() const
 
 void UFractureToolCluster::RegisterUICommand( FFractureEditorCommands* BindingContext ) 
 {
-	UI_COMMAND_EXT( BindingContext, 
-					UICommandInfo,
-					"Clustered", 
-					"Clustered", //GetDisplayText(),
-					"Clustered Voronoi Fracture", // GetTooltipText(),
-					EUserInterfaceActionType::ToggleButton, 
-					FInputChord() 
-	);
-
+	UI_COMMAND_EXT( BindingContext, UICommandInfo, "Clustered", "Clustered", "Clustered Voronoi Fracture", EUserInterfaceActionType::ToggleButton, FInputChord() );
 	BindingContext->Clustered = UICommandInfo;
 }
 
diff --git a/Engine/Plugins/Experimental/ChaosEditor/Source/FractureEditor/Private/FractureToolUniform.cpp b/Engine/Plugins/Experimental/ChaosEditor/Source/FractureEditor/Private/FractureToolUniform.cpp
--- a/Engine/Plugins/Experimental/ChaosEditor/Source/FractureEditor/Private/FractureToolUniform.cpp
+++ b/Engine/Plugins/Experimental/ChaosEditor/Source/FractureEditor/Private/FractureToolUniform.cpp
@@ -50,17 +50,8 @@ FSlateIcon UFractureToolUniform::GetToolIcon() const
 
 void UFractureToolUniform::RegisterUICommand( FFractureEditorCommands* BindingContext ) 
 {
-	UI_COMMAND_EXT( BindingContext, 
-					UICommandInfo,
-					"Uniform", 
-					"Uniform", //GetDisplayText(),
-					"Uniform Voronoi Fracture", // GetTooltipText(),
-					EUserInterfaceActionType::ToggleButton, 
-					FInputChord() 
-	);
-
+	UI_COMMAND_EXT( BindingContext, UICommandInfo, "Uniform", "Uniform", "Uniform Voronoi Fracture", EUserInterfaceActionType::ToggleButton, FInputChord() );
 	BindingContext->Uniform = UICommandInfo;
-
 }
 
 TArray<UObject*> UFractureToolUniform::GetSettingsObjects() const 
diff --git a/Engine/Plugins/Experimental/ControlRig/Source/ControlRigEditor/Private/Graph/SControlRigGraphNode.cpp b/Engine/Plugins/Experimental/ControlRig/Source/ControlRigEditor/Private/Graph/SControlRigGraphNode.cpp
--- a/Engine/Plugins/Experimental/ControlRig/Source/ControlRigEditor/Private/Graph/SControlRigGraphNode.cpp
+++ b/Engine/Plugins/Experimental/ControlRig/Source/ControlRigEditor/Private/Graph/SControlRigGraphNode.cpp
@@ -759,7 +759,7 @@ void SControlRigGraphNode::GetNodeInfoPopups(FNodeInfoContext* Context, TArray<F
 					}
 					else
 					{
-						PinnedWatchText += FText::Format(LOCTEXT("WatchingAndValidFmt", "Invalid Property {0}"), FText::FromString(PinName)).ToString();//@TODO: Print out object being debugged name?
+						PinnedWatchText += FText::Format(LOCTEXT("InvalidPropertyFmt", "Invalid Property {0}"), FText::FromString(PinName)).ToString();//@TODO: Print out object being debugged name?
 					}
 
 					ValidWatchCount++;
diff --git a/Engine/Plugins/Experimental/DataPrepEditor/Source/DataPrepEditor/Private/SchemaActions/DataprepDragDropOp.cpp b/Engine/Plugins/Experimental/DataPrepEditor/Source/DataPrepEditor/Private/SchemaActions/DataprepDragDropOp.cpp
--- a/Engine/Plugins/Experimental/DataPrepEditor/Source/DataPrepEditor/Private/SchemaActions/DataprepDragDropOp.cpp
+++ b/Engine/Plugins/Experimental/DataPrepEditor/Source/DataPrepEditor/Private/SchemaActions/DataprepDragDropOp.cpp
@@ -134,7 +134,7 @@ bool FDataprepDragDropOp::DoDropOnDataprepActionContext(const FDataprepSchemaAct
 {
 	if ( DataprepGraphOperation.IsBound() )
 	{
-		FScopedTransaction Transaction( LOCTEXT("AddNode", "Add a Step to a Dataprep Action") );
+		FScopedTransaction Transaction( LOCTEXT("AddStep", "Add a Step to a Dataprep Action") );
 		bool bDidModification = DataprepGraphOperation.Execute( Context );
 		if ( !bDidModification )
 		{
diff --git a/Engine/Plugins/Experimental/DataPrepEditor/Source/DataPrepLibraries/Private/DataprepOperations.cpp b/Engine/Plugins/Experimental/DataPrepEditor/Source/DataPrepLibraries/Private/DataprepOperations.cpp
--- a/Engine/Plugins/Experimental/DataPrepEditor/Source/DataPrepLibraries/Private/DataprepOperations.cpp
+++ b/Engine/Plugins/Experimental/DataPrepEditor/Source/DataPrepLibraries/Private/DataprepOperations.cpp
@@ -318,8 +318,8 @@ void FDataprepSetLOGGroupDetails::CustomizeDetails(IDetailLayoutBuilder & Detail
 	CustomAssetImportRow.NameContent()
 	[
 		SNew(STextBlock)
-		.Text(LOCTEXT("DatasmithMeshOperations", "LODGroupName"))
-		.ToolTipText(LOCTEXT("DatasmithMeshOperations", "List of predefined LODGroup"))
+		.Text(LOCTEXT("DatasmithMeshOperationsLabel", "LODGroupName"))
+		.ToolTipText(LOCTEXT("DatasmithMeshOperationsTooltip", "List of predefined LODGroup"))
 		.Font( DetailBuilder.GetDetailFont() )
 	];
 
diff --git a/Engine/Plugins/Experimental/PanoramicCapture/Source/PanoramicCapture/Private/SceneCapturer.cpp b/Engine/Plugins/Experimental/PanoramicCapture/Source/PanoramicCapture/Private/SceneCapturer.cpp
--- a/Engine/Plugins/Experimental/PanoramicCapture/Source/PanoramicCapture/Private/SceneCapturer.cpp
+++ b/Engine/Plugins/Experimental/PanoramicCapture/Source/PanoramicCapture/Private/SceneCapturer.cpp
@@ -649,7 +649,7 @@ void USceneCapturer::ValidateParameters()
 		if (!FPaths::DirectoryExists(Drive))
 		{
 			ErrorFound = true;
-			FMessageLog(StereoPanoramaLogName).Message(EMessageSeverity::Error, LOCTEXT("ValidationError1", "The output directory's drive doesn't exists. Plese set SP.OutputDir with a valid path. Skipping renders..."));
+			FMessageLog(StereoPanoramaLogName).Message(EMessageSeverity::Error, LOCTEXT("ValidationError_MissingOutputDirectory", "The output directory's drive doesn't exists. Plese set SP.OutputDir with a valid path. Skipping renders..."));
 		}
 	}
 
@@ -657,85 +657,85 @@ void USceneCapturer::ValidateParameters()
 	if (MaterialBlendableAO == NULL || MaterialBlendableBaseColor == NULL || MaterialBlendableMetallic == NULL || MaterialBlendableRoughness == NULL || MaterialBlendableWorldNormal == NULL)
 	{
 		ErrorFound = true;
-		FMessageLog(StereoPanoramaLogName).Message(EMessageSeverity::Error, LOCTEXT("ValidationError2", "Missing blendable materials. Skipping renders..."));
+		FMessageLog(StereoPanoramaLogName).Message(EMessageSeverity::Error, LOCTEXT("ValidationError_MissingBlendableMaterials", "Missing blendable materials. Skipping renders..."));
 	}
 
 	// Angular increment needs to be a factor of 360 to avoid seams i.e. 360 / angular increment needs to be a whole number
 	if ((int32)(NumberOfHorizontalSteps * hAngIncrement) != 360)
 	{
-		FMessageLog(StereoPanoramaLogName).Message(EMessageSeverity::Warning, FText::Format(LOCTEXT("ValidationWarning1", "Horizontal angular step {hAngIncrement} is not a factor of 360! This will lead to a seam between the start and end points"), Args));
+		FMessageLog(StereoPanoramaLogName).Message(EMessageSeverity::Warning, FText::Format(LOCTEXT("ValidationWarning_InvalidHorizonalAngularStep", "Horizontal angular step {hAngIncrement} is not a factor of 360! This will lead to a seam between the start and end points"), Args));
 	}
 
 	if ((int32)((NumberOfVerticalSteps - 1) * vAngIncrement) != 180)
 	{
-		FMessageLog(StereoPanoramaLogName).Message(EMessageSeverity::Warning, FText::Format(LOCTEXT("ValidationWarning2", "Vertical angular step {vAngIncrement} is not a factor of 180! This will lead to a seam between the start and end points"), Args));
+		FMessageLog(StereoPanoramaLogName).Message(EMessageSeverity::Warning, FText::Format(LOCTEXT("ValidationWarning_InvalidVerticalAngularStep", "Vertical angular step {vAngIncrement} is not a factor of 180! This will lead to a seam between the start and end points"), Args));
 	}
 
 	TotalSteps = NumberOfHorizontalSteps * NumberOfVerticalSteps;
 	if ((SphericalAtlasWidth & 1) != 0)
 	{
-		FMessageLog(StereoPanoramaLogName).Message(EMessageSeverity::Warning, FText::Format(LOCTEXT("ValidationWarning3", "The Atlas Width {SphericalAtlasWidth} must be even! Otherwise the Atlas height will not divide evenly."), Args));
+		FMessageLog(StereoPanoramaLogName).Message(EMessageSeverity::Warning, FText::Format(LOCTEXT("ValidationWarning_InvalidAtlasWidth", "The Atlas Width {SphericalAtlasWidth} must be even! Otherwise the Atlas height will not divide evenly."), Args));
 	}
 
 	// The strip width needs to be an even number and a factor of the number of steps
 	if ((StripWidth & 1) != 0)
 	{
-		FMessageLog(StereoPanoramaLogName).Message(EMessageSeverity::Warning, FText::Format(LOCTEXT("ValidationWarning4", "Strip width {StripWidth} needs to be even to avoid bad offsets"), Args));
+		FMessageLog(StereoPanoramaLogName).Message(EMessageSeverity::Warning, FText::Format(LOCTEXT("ValidationWarning_InvalidStripWidth", "Strip width {StripWidth} needs to be even to avoid bad offsets"), Args));
 	}
 
 	if ((StripHeight & 1) != 0)
 	{
-		FMessageLog(StereoPanoramaLogName).Message(EMessageSeverity::Warning, FText::Format(LOCTEXT("ValidationWarning5", "Strip height {StripHeight} needs to be even to avoid bad offsets"), Args));
+		FMessageLog(StereoPanoramaLogName).Message(EMessageSeverity::Warning, FText::Format(LOCTEXT("ValidationWarning_InvalidStripHeight", "Strip height {StripHeight} needs to be even to avoid bad offsets"), Args));
 	}
 
 	// Validate capturewidth & captureheight. Need to be even
 	if (CaptureWidth % 2 != 0)
 	{
-		FMessageLog(StereoPanoramaLogName).Message(EMessageSeverity::Warning, FText::Format(LOCTEXT("ValidationWarning6", "The {CaptureWidth} needs to be an even number"), Args));
+		FMessageLog(StereoPanoramaLogName).Message(EMessageSeverity::Warning, FText::Format(LOCTEXT("ValidationWarning_InvalidCaptureWidth", "The {CaptureWidth} needs to be an even number"), Args));
 	}
 
 	if (CaptureHeight % 2 != 0)
 	{
-		FMessageLog(StereoPanoramaLogName).Message(EMessageSeverity::Warning, FText::Format(LOCTEXT("ValidationWarning7", "The {CaptureHeight} needs to be an even number"), Args));
+		FMessageLog(StereoPanoramaLogName).Message(EMessageSeverity::Warning, FText::Format(LOCTEXT("ValidationWarning_InvalidCaptureHeight", "The {CaptureHeight} needs to be an even number"), Args));
 
 	}
 
 	// Unnecessary warning
 	//if (StripWidth * NumberOfHorizontalSteps != SphericalAtlasWidth)
 	//{
-		//FMessageLog(StereoPanoramaLogName).Message(EMessageSeverity::Warning, FText::Format(LOCTEXT("ValidationWarning5", "The number of horizontal steps {NumberOfHorizontalSteps} needs to be a factor of the atlas width {SphericalAtlasWidth}"), Args));
+		//FMessageLog(StereoPanoramaLogName).Message(EMessageSeverity::Warning, FText::Format(LOCTEXT("ValidationWarning_InvalidNumHorizonalSteps", "The number of horizontal steps {NumberOfHorizontalSteps} needs to be a factor of the atlas width {SphericalAtlasWidth}"), Args));
 	//}
 
 	// Unnecessary warning
 	//if (StripHeight * (NumberOfVerticalSteps - 1) != SphericalAtlasHeight)
 	//{
-		//FMessageLog(StereoPanoramaLogName).Message(EMessageSeverity::Warning, FText::Format(LOCTEXT("ValidationWarning7", "The number of vertical steps {NumberOfVerticalSteps} needs to be a factor of the atlas height {SphericalAtlasHeight}"), Args));
+		//FMessageLog(StereoPanoramaLogName).Message(EMessageSeverity::Warning, FText::Format(LOCTEXT("ValidationWarning_InvalidNumVerticalSteps", "The number of vertical steps {NumberOfVerticalSteps} needs to be a factor of the atlas height {SphericalAtlasHeight}"), Args));
 	//}
 
-	FMessageLog(StereoPanoramaLogName).Message(EMessageSeverity::Info, LOCTEXT("ValidationInfo0", "Panoramic screenshot parameters"));
+	FMessageLog(StereoPanoramaLogName).Message(EMessageSeverity::Info, LOCTEXT("ValidationInfo_PanoramicScreenshotParams", "Panoramic screenshot parameters"));
 
 	if (bMonoscopicMode)
 	{
-		FMessageLog(StereoPanoramaLogName).Message(EMessageSeverity::Info, LOCTEXT("ValidationInfo1", " ... In Monoscopic mode"));
+		FMessageLog(StereoPanoramaLogName).Message(EMessageSeverity::Info, LOCTEXT("ValidationInfo_MonoscopicMode", " ... In Monoscopic mode"));
 	}
 	else
 	{
-		FMessageLog(StereoPanoramaLogName).Message(EMessageSeverity::Info, LOCTEXT("ValidationInfo1", " ... In Stereoscopic mode"));
+		FMessageLog(StereoPanoramaLogName).Message(EMessageSeverity::Info, LOCTEXT("ValidationInfo_StereoscopicMode", " ... In Stereoscopic mode"));
 	}
 
-	FMessageLog(StereoPanoramaLogName).Message(EMessageSeverity::Info, FText::Format(LOCTEXT("ValidationInfo2", " ... capture size: {CaptureWidth} x {CaptureHeight}"), Args));
-	FMessageLog(StereoPanoramaLogName).Message(EMessageSeverity::Info, FText::Format(LOCTEXT("ValidationInfo3", " ... spherical atlas size: {SphericalAtlasWidth} x {SphericalAtlasHeight}"), Args));
-	FMessageLog(StereoPanoramaLogName).Message(EMessageSeverity::Info, FText::Format(LOCTEXT("ValidationInfo4", " ... intermediate atlas size: {UnprojectedAtlasWidth} x {UnprojectedAtlasHeight}"), Args));
-	FMessageLog(StereoPanoramaLogName).Message(EMessageSeverity::Info, FText::Format(LOCTEXT("ValidationInfo5", " ... strip size: {StripWidth} x {StripHeight}"), Args));
-	FMessageLog(StereoPanoramaLogName).Message(EMessageSeverity::Info, FText::Format(LOCTEXT("ValidationInfo6", " ... horizontal steps: {NumberOfHorizontalSteps} at {hAngIncrement} degrees"), Args));
-	FMessageLog(StereoPanoramaLogName).Message(EMessageSeverity::Info, FText::Format(LOCTEXT("ValidationInfo7", " ... vertical steps: {NumberOfVerticalSteps} at {vAngIncrement} degrees"), Args));
+	FMessageLog(StereoPanoramaLogName).Message(EMessageSeverity::Info, FText::Format(LOCTEXT("ValidationInfo_CaptureSize", " ... capture size: {CaptureWidth} x {CaptureHeight}"), Args));
+	FMessageLog(StereoPanoramaLogName).Message(EMessageSeverity::Info, FText::Format(LOCTEXT("ValidationInfo_SphericalAtlasSize", " ... spherical atlas size: {SphericalAtlasWidth} x {SphericalAtlasHeight}"), Args));
+	FMessageLog(StereoPanoramaLogName).Message(EMessageSeverity::Info, FText::Format(LOCTEXT("ValidationInfo_IntermediateAtlasSize", " ... intermediate atlas size: {UnprojectedAtlasWidth} x {UnprojectedAtlasHeight}"), Args));
+	FMessageLog(StereoPanoramaLogName).Message(EMessageSeverity::Info, FText::Format(LOCTEXT("ValidationInfo_StripSize", " ... strip size: {StripWidth} x {StripHeight}"), Args));
+	FMessageLog(StereoPanoramaLogName).Message(EMessageSeverity::Info, FText::Format(LOCTEXT("ValidationInfo_NumHorizonalSteps", " ... horizontal steps: {NumberOfHorizontalSteps} at {hAngIncrement} degrees"), Args));
+	FMessageLog(StereoPanoramaLogName).Message(EMessageSeverity::Info, FText::Format(LOCTEXT("ValidationInfo_NumVerticalSteps", " ... vertical steps: {NumberOfVerticalSteps} at {vAngIncrement} degrees"), Args));
 }
 
 void USceneCapturer::SetInitialState( int32 InStartFrame, int32 InEndFrame, FStereoCaptureDoneDelegate& InStereoCaptureDoneDelegate )
 {
 	if( bIsTicking )
 	{
-		FMessageLog(StereoPanoramaLogName).Message(EMessageSeverity::Warning, LOCTEXT("InitialStateWarning1", "Already capturing a scene; concurrent captures are not allowed"));
+		FMessageLog(StereoPanoramaLogName).Message(EMessageSeverity::Warning, LOCTEXT("InitialStateWarning_AlreadyCapturing", "Already capturing a scene; concurrent captures are not allowed"));
 		return;
 	}
 
@@ -744,7 +744,7 @@ void USceneCapturer::SetInitialState( int32 InStartFrame, int32 InEndFrame, FSte
 
 	if( CaptureGameMode == NULL || CapturePlayerController == NULL )
 	{
-		FMessageLog(StereoPanoramaLogName).Message(EMessageSeverity::Warning, LOCTEXT("InitialStateWarning2", "Missing GameMode or PlayerController"));
+		FMessageLog(StereoPanoramaLogName).Message(EMessageSeverity::Warning, LOCTEXT("InitialStateWarning_MissingGameModeOrPC", "Missing GameMode or PlayerController"));
 		return;
 	}
 
@@ -1268,7 +1268,7 @@ void USceneCapturer::Tick( float DeltaTime )
 				const TArray<uint8> ImageData = ImageWrapper->GetCompressed(100);
 				FFileHelper::SaveArrayToFile(ImageData, *AtlasName);
 				ImageWrapper.Reset();
-				FMessageLog(StereoPanoramaLogName).Message(EMessageSeverity::Info, LOCTEXT("Done8bit", "Done!"));
+				FMessageLog(StereoPanoramaLogName).Message(EMessageSeverity::Info, LOCTEXT("Done", "Done!"));
 			}			
 			else
 			{
@@ -1284,7 +1284,7 @@ void USceneCapturer::Tick( float DeltaTime )
 				const TArray<uint8> ImageData = ImageWrapper->GetCompressed((int32)EImageCompressionQuality::Default);
 				FFileHelper::SaveArrayToFile(ImageData, *AtlasName);
 				ImageWrapper.Reset();
-				FMessageLog(StereoPanoramaLogName).Message(EMessageSeverity::Info, LOCTEXT("Done32bit", "Done!"));
+				FMessageLog(StereoPanoramaLogName).Message(EMessageSeverity::Info, LOCTEXT("Done", "Done!"));
 			}
 
 		}
diff --git a/Engine/Plugins/FX/Niagara/Source/Niagara/Private/NDISkeletalMesh_BoneSampling.cpp b/Engine/Plugins/FX/Niagara/Source/Niagara/Private/NDISkeletalMesh_BoneSampling.cpp
--- a/Engine/Plugins/FX/Niagara/Source/Niagara/Private/NDISkeletalMesh_BoneSampling.cpp
+++ b/Engine/Plugins/FX/Niagara/Source/Niagara/Private/NDISkeletalMesh_BoneSampling.cpp
@@ -72,7 +72,7 @@ void UNiagaraDataInterfaceSkeletalMesh::GetSkeletonSamplingFunctions(TArray<FNia
 		Sig.bMemberFunction = true;
 		Sig.bRequiresContext = false;
 #if WITH_EDITORONLY_DATA
-		Sig.Description = LOCTEXT("GetSkinnedBoneDataDesc", "Returns skinning dependant data for the pased bone in local space. All outputs are optional and you will incur zero to minimal cost if they are not connected.");
+		Sig.Description = LOCTEXT("GetOptionalSkinnedBoneDataDesc", "Returns skinning dependant data for the pased bone in local space. All outputs are optional and you will incur zero to minimal cost if they are not connected.");
 #endif
 		OutFunctions.Add(Sig);
 	}
@@ -88,7 +88,7 @@ void UNiagaraDataInterfaceSkeletalMesh::GetSkeletonSamplingFunctions(TArray<FNia
 		Sig.bMemberFunction = true;
 		Sig.bRequiresContext = false;
 #if WITH_EDITORONLY_DATA
-		Sig.Description = LOCTEXT("GetSkinnedBoneDataWSDesc", "Returns skinning dependant data for the pased bone in world space. All outputs are optional and you will incur zero to minimal cost if they are not connected.");
+		Sig.Description = LOCTEXT("GetOptionalSkinnedBoneDataWSDesc", "Returns skinning dependant data for the pased bone in world space. All outputs are optional and you will incur zero to minimal cost if they are not connected.");
 #endif
 		OutFunctions.Add(Sig);
 	}
diff --git a/Engine/Plugins/FX/Niagara/Source/Niagara/Private/NDISkeletalMesh_TriangleSampling.cpp b/Engine/Plugins/FX/Niagara/Source/Niagara/Private/NDISkeletalMesh_TriangleSampling.cpp
--- a/Engine/Plugins/FX/Niagara/Source/Niagara/Private/NDISkeletalMesh_TriangleSampling.cpp
+++ b/Engine/Plugins/FX/Niagara/Source/Niagara/Private/NDISkeletalMesh_TriangleSampling.cpp
@@ -76,7 +76,7 @@ void UNiagaraDataInterfaceSkeletalMesh::GetTriangleSamplingFunctions(TArray<FNia
 		Sig.bMemberFunction = true;
 		Sig.bRequiresContext = false;
 #if WITH_EDITORONLY_DATA
-		Sig.Description = LOCTEXT("GetSkinnedDataDesc", "Returns skinning dependant data for the pased MeshTriCoord in local space. All outputs are optional and you will incur zerp minimal cost if they are not connected.");
+		Sig.Description = LOCTEXT("GetOptionalSkinnedDataDesc", "Returns skinning dependant data for the pased MeshTriCoord in local space. All outputs are optional and you will incur zerp minimal cost if they are not connected.");
 #endif
 		OutFunctions.Add(Sig);
 	}
@@ -94,7 +94,7 @@ void UNiagaraDataInterfaceSkeletalMesh::GetTriangleSamplingFunctions(TArray<FNia
 		Sig.bMemberFunction = true;
 		Sig.bRequiresContext = false;
 #if WITH_EDITORONLY_DATA
-		Sig.Description = LOCTEXT("GetSkinnedDataWSDesc", "Returns skinning dependant data for the pased MeshTriCoord in world space. All outputs are optional and you will incur zerp minimal cost if they are not connected.");
+		Sig.Description = LOCTEXT("GetOptionalSkinnedDataWSDesc", "Returns skinning dependant data for the pased MeshTriCoord in world space. All outputs are optional and you will incur zerp minimal cost if they are not connected.");
 #endif
 		OutFunctions.Add(Sig);
 	}
diff --git a/Engine/Plugins/FX/Niagara/Source/NiagaraEditor/Private/EdGraphSchema_Niagara.cpp b/Engine/Plugins/FX/Niagara/Source/NiagaraEditor/Private/EdGraphSchema_Niagara.cpp
--- a/Engine/Plugins/FX/Niagara/Source/NiagaraEditor/Private/EdGraphSchema_Niagara.cpp
+++ b/Engine/Plugins/FX/Niagara/Source/NiagaraEditor/Private/EdGraphSchema_Niagara.cpp
@@ -853,7 +853,7 @@ TArray<TSharedPtr<FNiagaraSchemaAction_NewNode> > UEdGraphSchema_Niagara::GetGra
 
 	// Add static switch node
 	{
-		const FText UsageSelectorMenuDesc = LOCTEXT("NiagaraUsageSelectorMenuDesc", "Static Switch");
+		const FText UsageSelectorMenuDesc = LOCTEXT("NiagaraStaticSwitchMenuDesc", "Static Switch");
 		TSharedPtr<FNiagaraSchemaAction_NewNode> Action = AddNewNodeAction(NewActions, UtilMenuCat, UsageSelectorMenuDesc, TEXT("Static Switch"), FText::GetEmpty());
 		UNiagaraNodeStaticSwitch* Node = NewObject<UNiagaraNodeStaticSwitch>(OwnerOfTemporaries);
 		Action->NodeTemplate = Node;
diff --git a/Engine/Plugins/FX/Niagara/Source/NiagaraEditor/Private/ViewModels/Stack/NiagaraStackFunctionInputCollection.cpp b/Engine/Plugins/FX/Niagara/Source/NiagaraEditor/Private/ViewModels/Stack/NiagaraStackFunctionInputCollection.cpp
--- a/Engine/Plugins/FX/Niagara/Source/NiagaraEditor/Private/ViewModels/Stack/NiagaraStackFunctionInputCollection.cpp
+++ b/Engine/Plugins/FX/Niagara/Source/NiagaraEditor/Private/ViewModels/Stack/NiagaraStackFunctionInputCollection.cpp
@@ -329,7 +329,7 @@ void UNiagaraStackFunctionInputCollection::RefreshIssues(TArray<FName> Duplicate
 				Fixes.Add(ConvertInputOverrideFix);
 
 				// second possible fix: remove the override completely
-				Fixes.Add(GetNodeRemovalFix(OverridePin, LOCTEXT("RemoveInvalidInputTransaction", "Remove input override (WARNING: this could result in different behavior!)")));
+				Fixes.Add(GetNodeRemovalFix(OverridePin, LOCTEXT("RemoveInvalidInputTransactionExt", "Remove input override (WARNING: this could result in different behavior!)")));
 
 				FStackIssue DeprecatedInputOverrideError(
 					EStackIssueSeverity::Error,
diff --git a/Engine/Plugins/Runtime/AR/Microsoft/HoloLensAR/Source/HoloLensPlatformEditor/Private/HoloLensTargetSettingsCustomization.cpp b/Engine/Plugins/Runtime/AR/Microsoft/HoloLensAR/Source/HoloLensPlatformEditor/Private/HoloLensTargetSettingsCustomization.cpp
--- a/Engine/Plugins/Runtime/AR/Microsoft/HoloLensAR/Source/HoloLensPlatformEditor/Private/HoloLensTargetSettingsCustomization.cpp
+++ b/Engine/Plugins/Runtime/AR/Microsoft/HoloLensAR/Source/HoloLensPlatformEditor/Private/HoloLensTargetSettingsCustomization.cpp
@@ -705,7 +705,7 @@ void FHoloLensTargetSettingsCustomization::LoadAndValidateSigningCertificate()
 			switch (FailureReason)
 			{
 			case ERROR_INVALID_PASSWORD:
-				SigningCertificateError->SetError(LOCTEXT("CertificatePasswordProtected", "SigningCertificate.pfx file is password protected"));
+				SigningCertificateError->SetError(LOCTEXT("CertificateFilePasswordProtected", "SigningCertificate.pfx file is password protected"));
 				break;
 
 			default:
diff --git a/Engine/Plugins/Runtime/GameplayAbilities/Source/GameplayAbilities/Private/GameplayEffectTypes.cpp b/Engine/Plugins/Runtime/GameplayAbilities/Source/GameplayAbilities/Private/GameplayEffectTypes.cpp
--- a/Engine/Plugins/Runtime/GameplayAbilities/Source/GameplayAbilities/Private/GameplayEffectTypes.cpp
+++ b/Engine/Plugins/Runtime/GameplayAbilities/Source/GameplayAbilities/Private/GameplayEffectTypes.cpp
@@ -1097,3 +1097,5 @@ bool FMinimalReplicationTagCountMap::NetSerialize(FArchive& Ar, class UPackageMa
 	bOutSuccess = true;
 	return true;
 }
+
+#undef LOCTEXT_NAMESPACE
diff --git a/Engine/Source/Developer/AssetTools/Private/AssetVtConversion.cpp b/Engine/Source/Developer/AssetTools/Private/AssetVtConversion.cpp
--- a/Engine/Source/Developer/AssetTools/Private/AssetVtConversion.cpp
+++ b/Engine/Source/Developer/AssetTools/Private/AssetVtConversion.cpp
@@ -477,7 +477,7 @@ void FVTConversionWorker::FindAllTexturesAndMaterials(TArray<UMaterial *> &OutAf
 
 void FVTConversionWorker::FilterList(int32 SizeThreshold)
 {
-	FScopedSlowTask SlowTask(1.0f, LOCTEXT("ConvertToVT_Progress_ConvertToVt", "Finding textures to convert..."));
+	FScopedSlowTask SlowTask(1.0f, LOCTEXT("ConvertToVT_Progress_FindingTextures", "Finding textures to convert..."));
 	SlowTask.MakeDialog();
 
 	Textures.Empty();
@@ -520,7 +520,7 @@ void FVTConversionWorker::DoConvert(bool backwards)
 {
 	const bool bVirtualTextureEnable = !backwards;
 
-	FScopedSlowTask SlowTask(2.0f, LOCTEXT("ConvertToVT_Progress_ConvertToVt", "Converting textures and materials..."));
+	FScopedSlowTask SlowTask(2.0f, LOCTEXT("ConvertToVT_Progress_ConvertingTexturesAndMaterials", "Converting textures and materials..."));
 	SlowTask.MakeDialog();
 
 	TMap<UTexture2D*, UTexture2D*> EngineTextureToCopyMap;
diff --git a/Engine/Source/Developer/TraceInsights/Private/Insights/Widgets/SStartPageWindow.cpp b/Engine/Source/Developer/TraceInsights/Private/Insights/Widgets/SStartPageWindow.cpp
--- a/Engine/Source/Developer/TraceInsights/Private/Insights/Widgets/SStartPageWindow.cpp
+++ b/Engine/Source/Developer/TraceInsights/Private/Insights/Widgets/SStartPageWindow.cpp
@@ -1086,7 +1086,7 @@ FText SStartPageWindow::GetRecorderStatusText() const
 	}
 	else
 	{
-		return FText(LOCTEXT("RecorderServerRunning", "Stopped"));
+		return FText(LOCTEXT("RecorderServerStopped", "Stopped"));
 	}
 }
 
diff --git a/Engine/Source/Editor/Blutility/Private/AssetTypeActions_EditorUtilityWidgetBlueprint.cpp b/Engine/Source/Editor/Blutility/Private/AssetTypeActions_EditorUtilityWidgetBlueprint.cpp
--- a/Engine/Source/Editor/Blutility/Private/AssetTypeActions_EditorUtilityWidgetBlueprint.cpp
+++ b/Engine/Source/Editor/Blutility/Private/AssetTypeActions_EditorUtilityWidgetBlueprint.cpp
@@ -98,12 +98,12 @@ void FAssetTypeActions_EditorUtilityWidgetBlueprint::PerformAssetDiff(UObject* A
 	// times we're comparing two completely separate assets altogether)
 	bool bIsSingleAsset = (NewBlueprint->GetName() == OldBlueprint->GetName());
 
-	FText WindowTitle = LOCTEXT("NamelessWidgetBlueprintDiff", "Editor Utility Widget Blueprint Diff");
+	FText WindowTitle = LOCTEXT("NamelessEditorUtilityWidgetBlueprintDiff", "Editor Utility Widget Blueprint Diff");
 	// if we're diffing one asset against itself 
 	if (bIsSingleAsset)
 	{
 		// identify the assumed single asset in the window's title
-		WindowTitle = FText::Format(LOCTEXT("WidgetBlueprintDiff", "{0} - Editor Utility Widget Blueprint Diff"), FText::FromString(NewBlueprint->GetName()));
+		WindowTitle = FText::Format(LOCTEXT("EditorUtilityWidgetBlueprintDiff", "{0} - Editor Utility Widget Blueprint Diff"), FText::FromString(NewBlueprint->GetName()));
 	}
 
 	SBlueprintDiff::CreateDiffWindow(WindowTitle, OldBlueprint, NewBlueprint, OldRevision, NewRevision);
diff --git a/Engine/Source/Editor/Experimental/EditorInteractiveToolsFramework/Private/StandardToolModeCommands.cpp b/Engine/Source/Editor/Experimental/EditorInteractiveToolsFramework/Private/StandardToolModeCommands.cpp
--- a/Engine/Source/Editor/Experimental/EditorInteractiveToolsFramework/Private/StandardToolModeCommands.cpp
+++ b/Engine/Source/Editor/Experimental/EditorInteractiveToolsFramework/Private/StandardToolModeCommands.cpp
@@ -36,3 +36,4 @@ TSharedPtr<FUICommandInfo> FStandardToolModeCommands::FindStandardCommand(EStand
 }
 
 
+#undef LOCTEXT_NAMESPACE
diff --git a/Engine/Source/Editor/LandscapeEditor/Private/LandscapeEditorDetailCustomization_Layers.cpp b/Engine/Source/Editor/LandscapeEditor/Private/LandscapeEditorDetailCustomization_Layers.cpp
--- a/Engine/Source/Editor/LandscapeEditor/Private/LandscapeEditorDetailCustomization_Layers.cpp
+++ b/Engine/Source/Editor/LandscapeEditor/Private/LandscapeEditorDetailCustomization_Layers.cpp
@@ -484,18 +484,18 @@ TSharedPtr<SWidget> FLandscapeEditorCustomNodeBuilder_Layers::OnLayerContextMenu
 				{
 					// Hide Selected Layer
 					FUIAction HideSelectedLayerAction = FUIAction(FExecuteAction::CreateLambda([SharedThis, InLayerIndex] { SharedThis->OnToggleVisibility(InLayerIndex); }));
-					MenuBuilder.AddMenuEntry(LOCTEXT("HideSelectedLayer", "Hide Selected"), LOCTEXT("HideSelectedLayerTooltip", "Hide Selected Layer"), FSlateIcon(), HideSelectedLayerAction);
+					MenuBuilder.AddMenuEntry(LOCTEXT("HideSelected", "Hide Selected"), LOCTEXT("HideSelectedLayerTooltip", "Hide Selected Layer"), FSlateIcon(), HideSelectedLayerAction);
 				}
 				else
 				{
 					// Show Selected Layer
 					FUIAction ShowSelectedLayerAction = FUIAction(FExecuteAction::CreateLambda([SharedThis, InLayerIndex] { SharedThis->OnToggleVisibility(InLayerIndex); }));
-					MenuBuilder.AddMenuEntry(LOCTEXT("ShowSelectedLayer", "Show Selected"), LOCTEXT("ShowSelectedLayerTooltip", "Show Selected Layer"), FSlateIcon(), ShowSelectedLayerAction);
+					MenuBuilder.AddMenuEntry(LOCTEXT("ShowSelected", "Show Selected"), LOCTEXT("ShowSelectedLayerTooltip", "Show Selected Layer"), FSlateIcon(), ShowSelectedLayerAction);
 				}
 
 				// Show Only Selected Layer
 				FUIAction ShowOnlySelectedLayerAction = FUIAction(FExecuteAction::CreateLambda([SharedThis, InLayerIndex] { SharedThis->ShowOnlySelectedLayer(InLayerIndex); }));
-				MenuBuilder.AddMenuEntry(LOCTEXT("ShowOnlySelectedLayer", "Show Only Selected"), LOCTEXT("ShowOnlySelectedLayerTooltip", "Show Only Selected Layer"), FSlateIcon(), ShowOnlySelectedLayerAction);
+				MenuBuilder.AddMenuEntry(LOCTEXT("ShowOnlySelected", "Show Only Selected"), LOCTEXT("ShowOnlySelectedLayerTooltip", "Show Only Selected Layer"), FSlateIcon(), ShowOnlySelectedLayerAction);
 			}
 
 			// Show All Layers
diff --git a/Engine/Source/Editor/MovieSceneTools/Private/TrackEditors/ParticleParameterTrackEditor.cpp b/Engine/Source/Editor/MovieSceneTools/Private/TrackEditors/ParticleParameterTrackEditor.cpp
--- a/Engine/Source/Editor/MovieSceneTools/Private/TrackEditors/ParticleParameterTrackEditor.cpp
+++ b/Engine/Source/Editor/MovieSceneTools/Private/TrackEditors/ParticleParameterTrackEditor.cpp
@@ -53,7 +53,7 @@ void FParticleParameterTrackEditor::BuildObjectBindingTrackMenu( FMenuBuilder& M
 		const TSharedPtr<ISequencer> ParentSequencer = GetSequencer();
 
 		MenuBuilder.AddMenuEntry(
-			LOCTEXT( "AddParticleParameterTrack", "Particle Parameter Track" ),
+			LOCTEXT( "AddParticleParameterTrackLabel", "Particle Parameter Track" ),
 			LOCTEXT( "AddParticleParameterTrackTooltip", "Adds a track for controlling particle parameter values." ),
 			FSlateIcon(),
 			FUIAction
diff --git a/Engine/Source/Editor/PropertyEditor/Private/UserInterface/PropertyEditor/SPropertyMenuComponentPicker.cpp b/Engine/Source/Editor/PropertyEditor/Private/UserInterface/PropertyEditor/SPropertyMenuComponentPicker.cpp
--- a/Engine/Source/Editor/PropertyEditor/Private/UserInterface/PropertyEditor/SPropertyMenuComponentPicker.cpp
+++ b/Engine/Source/Editor/PropertyEditor/Private/UserInterface/PropertyEditor/SPropertyMenuComponentPicker.cpp
@@ -39,22 +39,22 @@ void SPropertyMenuComponentPicker::Construct(const FArguments& InArgs)
 		if (InitialComponent)
 		{
 			MenuBuilder.AddMenuEntry(
-				LOCTEXT("EditAsset", "Edit"),
-				LOCTEXT("EditAsset_Tooltip", "Edit this component"),
+				LOCTEXT("EditComponent", "Edit"),
+				LOCTEXT("EditComponent_Tooltip", "Edit this component"),
 				FSlateIcon(),
 				FUIAction(FExecuteAction::CreateSP(this, &SPropertyMenuComponentPicker::OnEdit)));
 		}
 
 		MenuBuilder.AddMenuEntry(
-			LOCTEXT("CopyAsset", "Copy"),
-			LOCTEXT("CopyAsset_Tooltip", "Copies the component to the clipboard"),
+			LOCTEXT("CopyComponent", "Copy"),
+			LOCTEXT("CopyComponent_Tooltip", "Copies the component to the clipboard"),
 			FSlateIcon(),
 			FUIAction(FExecuteAction::CreateSP(this, &SPropertyMenuComponentPicker::OnCopy))
 			);
 
 		MenuBuilder.AddMenuEntry(
-			LOCTEXT("PasteAsset", "Paste"),
-			LOCTEXT("PasteAsset_Tooltip", "Pastes an component from the clipboard to this field"),
+			LOCTEXT("PasteComponent", "Paste"),
+			LOCTEXT("PasteComponent_Tooltip", "Pastes an component from the clipboard to this field"),
 			FSlateIcon(),
 			FUIAction(
 				FExecuteAction::CreateSP(this, &SPropertyMenuComponentPicker::OnPaste),
@@ -64,8 +64,8 @@ void SPropertyMenuComponentPicker::Construct(const FArguments& InArgs)
 		if (bAllowClear)
 		{
 			MenuBuilder.AddMenuEntry(
-				LOCTEXT("ClearAsset", "Clear"),
-				LOCTEXT("ClearAsset_ToolTip", "Clears the component set on this field"),
+				LOCTEXT("ClearComponent", "Clear"),
+				LOCTEXT("ClearComponent_ToolTip", "Clears the component set on this field"),
 				FSlateIcon(),
 				FUIAction(FExecuteAction::CreateSP(this, &SPropertyMenuComponentPicker::OnClear))
 			);
diff --git a/Engine/Source/Editor/Sequencer/Private/DisplayNodes/SequencerObjectBindingNode.cpp b/Engine/Source/Editor/Sequencer/Private/DisplayNodes/SequencerObjectBindingNode.cpp
--- a/Engine/Source/Editor/Sequencer/Private/DisplayNodes/SequencerObjectBindingNode.cpp
+++ b/Engine/Source/Editor/Sequencer/Private/DisplayNodes/SequencerObjectBindingNode.cpp
@@ -220,7 +220,7 @@ void FSequencerObjectBindingNode::BuildContextMenu(FMenuBuilder& MenuBuilder)
 
 			MenuBuilder.AddMenuEntry(
 				LOCTEXT("ContinuouslyRespawn", "Continuously Respawn"),
-				LOCTEXT("ContinuouslyRespawn", "When enabled, this spawnable will always be respawned if it gets destroyed externally. When disabled, this object will only ever be spawned once for each spawn key even if it is destroyed externally"),
+				LOCTEXT("ContinuouslyRespawnTooltip", "When enabled, this spawnable will always be respawned if it gets destroyed externally. When disabled, this object will only ever be spawned once for each spawn key even if it is destroyed externally"),
 				FSlateIcon(),
 				FUIAction(FExecuteAction::CreateLambda(ToggleContinuouslyRespawn), FCanExecuteAction(), FGetActionCheckState::CreateLambda(ContinuouslyRespawnCheckState)),
 				NAME_None,
@@ -241,7 +241,7 @@ void FSequencerObjectBindingNode::BuildContextMenu(FMenuBuilder& MenuBuilder)
 				FFormatNamedArguments Args;
 
 				MenuBuilder.AddSubMenu(
-					FText::Format(LOCTEXT("Assign Actor", "Assign Actor"), Args),
+					FText::Format(LOCTEXT("AssignActor", "Assign Actor"), Args),
 					FText::Format(LOCTEXT("AssignActorTooltip", "Assign an actor to this track"), Args),
 					FNewMenuDelegate::CreateSP(this, &FSequencerObjectBindingNode::AddAssignActorMenu));
 			}
@@ -252,7 +252,7 @@ void FSequencerObjectBindingNode::BuildContextMenu(FMenuBuilder& MenuBuilder)
 		MenuBuilder.BeginSection("Import/Export", LOCTEXT("ImportExportMenuSectionName", "Import/Export"));
 		
 		MenuBuilder.AddMenuEntry(
-			LOCTEXT("Import FBX", "Import..."),
+			LOCTEXT("ImportFBX", "Import..."),
 			LOCTEXT("ImportFBXTooltip", "Import FBX animation to this object"),
 			FSlateIcon(),
 			FUIAction(
@@ -260,7 +260,7 @@ void FSequencerObjectBindingNode::BuildContextMenu(FMenuBuilder& MenuBuilder)
 			));
 
 		MenuBuilder.AddMenuEntry(
-			LOCTEXT("Export FBX", "Export..."),
+			LOCTEXT("ExportFBX", "Export..."),
 			LOCTEXT("ExportFBXTooltip", "Export FBX animation from this object"),
 			FSlateIcon(),
 			FUIAction(
@@ -268,7 +268,7 @@ void FSequencerObjectBindingNode::BuildContextMenu(FMenuBuilder& MenuBuilder)
 			));
 
 		MenuBuilder.AddMenuEntry(
-			LOCTEXT("Export to Camera Anim", "Export to Camera Anim..."),
+			LOCTEXT("ExportToCameraAnim", "Export to Camera Anim..."),
 			LOCTEXT("ExportToCameraAnimTooltip", "Exports the animation to a camera anim asset"),
 			FSlateIcon(),
 			FUIAction(
diff --git a/Engine/Source/Editor/UnrealEd/Private/Factories/PackFactory.cpp b/Engine/Source/Editor/UnrealEd/Private/Factories/PackFactory.cpp
--- a/Engine/Source/Editor/UnrealEd/Private/Factories/PackFactory.cpp
+++ b/Engine/Source/Editor/UnrealEd/Private/Factories/PackFactory.cpp
@@ -678,7 +678,7 @@ UObject* UPackFactory::FactoryCreateBinary
 					ILiveCodingModule* LiveCoding = FModuleManager::GetModulePtr<ILiveCodingModule>(LIVE_CODING_MODULE_NAME);
 					if (LiveCoding != nullptr && LiveCoding->IsEnabledForSession())
 					{
-						FMessageDialog::Open(EAppMsgType::Ok, NSLOCTEXT("PackFactory", "FailedToCompileNewGameModule", "Unable to compile source code while Live Coding is enabled. Please close the editor and build from your IDE."));
+						FMessageDialog::Open(EAppMsgType::Ok, NSLOCTEXT("PackFactory", "CannotCompileWithLiveCoding", "Unable to compile source code while Live Coding is enabled. Please close the editor and build from your IDE."));
 						bCompileSource = false;
 					}
 				}
diff --git a/Engine/Source/Editor/UnrealEd/Private/ObjectTools.cpp b/Engine/Source/Editor/UnrealEd/Private/ObjectTools.cpp
--- a/Engine/Source/Editor/UnrealEd/Private/ObjectTools.cpp
+++ b/Engine/Source/Editor/UnrealEd/Private/ObjectTools.cpp
@@ -2316,7 +2316,7 @@ namespace ObjectTools
 		FEditorDelegates::OnAssetsCanDelete.Broadcast(TArray<UObject*>{ ObjectToDelete }, CanDeleteResult);
 		if (!CanDeleteResult.Get())
 		{
-			FMessageDialog::Open(EAppMsgType::Ok, NSLOCTEXT("UnrealEd", "CannotDelete", "Cannot delete selected objects. See log for details."));
+			FMessageDialog::Open(EAppMsgType::Ok, NSLOCTEXT("UnrealEd", "CannotDelete", "Cannot currently delete selected objects. See log for details."));
 			return false;
 		}
 
