Stereo rendering refactor appears to have missed these.

diff --git a/Engine/Plugins/Experimental/PanoramicCapture/Source/PanoramicCapture/Private/StereoStaticMeshComponent.cpp b/Engine/Plugins/Experimental/PanoramicCapture/Source/PanoramicCapture/Private/StereoStaticMeshComponent.cpp
--- a/Engine/Plugins/Experimental/PanoramicCapture/Source/PanoramicCapture/Private/StereoStaticMeshComponent.cpp
+++ b/Engine/Plugins/Experimental/PanoramicCapture/Source/PanoramicCapture/Private/StereoStaticMeshComponent.cpp
@@ -28,14 +28,14 @@ public:
         FPrimitiveViewRelevance viewRelevance = FStaticMeshSceneProxy::GetViewRelevance(View);
         bool bVisible = true;
 
-		if (IStereoRendering::IsASecondaryView(View->StereoPass))
+		if (IStereoRendering::IsASecondaryView(*View, GEngine->StereoRenderingDevice))
 		{
 			if ((EyeToRender != ESPStereoCameraLayer::RightEye) && (EyeToRender != ESPStereoCameraLayer::BothEyes))
 			{
 				bVisible = false;
 			}
 		}
-		else if (IStereoRendering::IsAPrimaryView(View->StereoPass))
+		else if (IStereoRendering::IsAPrimaryView(*View, GEngine->StereoRenderingDevice))
 		{
 			if ((EyeToRender != ESPStereoCameraLayer::LeftEye) && (EyeToRender != ESPStereoCameraLayer::BothEyes))
 			{
diff --git a/Engine/Plugins/Lumin/MagicLeap/Source/MagicLeap/Private/MagicLeapXRCamera.cpp b/Engine/Plugins/Lumin/MagicLeap/Source/MagicLeap/Private/MagicLeapXRCamera.cpp
--- a/Engine/Plugins/Lumin/MagicLeap/Source/MagicLeap/Private/MagicLeapXRCamera.cpp
+++ b/Engine/Plugins/Lumin/MagicLeap/Source/MagicLeap/Private/MagicLeapXRCamera.cpp
@@ -16,9 +16,10 @@ void FMagicLeapXRCamera::PreRenderView_RenderThread(FRHICommandListImmediate& RH
 {
 #if WITH_MLSDK
 	// this needs to happen before the FDefaultXRCamera call, because UpdateProjectionMatrix is somewhat destructive. 
-	if (IStereoRendering::IsStereoEyeView(View.StereoPass))
+	if (IStereoRendering::IsStereoEyePass(View.StereoPass))
 	{
-		const int EyeIdx = IStereoRendering::IsAPrimaryView(View.StereoPass) ? 0 : 1;
+		const int EyeIdx = IStereoRendering::IsAPrimaryView(View, GEngine->StereoRenderingDevice) ? 0 : 1;
+
 		const FTrackingFrame& Frame = MagicLeapSystem.GetCurrentFrame();
 
 		// update to use render projection matrix
