Fix unity build error in StereoStaticMeshComponent.cpp

diff --git a/Engine/Plugins/Experimental/PanoramicCapture/Source/PanoramicCapture/Private/StereoStaticMeshComponent.cpp b/Engine/Plugins/Experimental/PanoramicCapture/Source/PanoramicCapture/Private/StereoStaticMeshComponent.cpp
--- a/Engine/Plugins/Experimental/PanoramicCapture/Source/PanoramicCapture/Private/StereoStaticMeshComponent.cpp
+++ b/Engine/Plugins/Experimental/PanoramicCapture/Source/PanoramicCapture/Private/StereoStaticMeshComponent.cpp
@@ -3,6 +3,7 @@
 #include "StereoStaticMeshComponent.h"
 #include "StaticMeshResources.h"
 #include "Engine/StaticMesh.h"
+#include "Engine/Engine.h"
 
 
 class FStereoStaticMeshSceneProxy final
