Fixes compilation failure in Composure

diff --git a/Engine/Plugins/Compositing/Composure/Source/Composure/Private/ComposureTonemapperPass.cpp b/Engine/Plugins/Compositing/Composure/Source/Composure/Private/ComposureTonemapperPass.cpp
--- a/Engine/Plugins/Compositing/Composure/Source/Composure/Private/ComposureTonemapperPass.cpp
+++ b/Engine/Plugins/Compositing/Composure/Source/Composure/Private/ComposureTonemapperPass.cpp
@@ -24,8 +24,8 @@ void FComposureTonemapperUtils::ApplyTonemapperSettings(const FColorGradingSetti
 		OutSettings.bOverride_VignetteIntensity = true;
 		OutSettings.VignetteIntensity = 0;
 
-		OutSettings.bOverride_GrainIntensity = true;
-		OutSettings.GrainIntensity = 0;
+		OutSettings.bOverride_FilmGrainIntensity = true;
+		OutSettings.FilmGrainIntensity = 0;
 
 		OutSettings.bOverride_BloomDirtMask = true;
 		OutSettings.BloomDirtMask = nullptr;
diff --git a/Engine/Plugins/Experimental/PanoramicCapture/Source/PanoramicCapture/Private/SceneCapturer.cpp b/Engine/Plugins/Experimental/PanoramicCapture/Source/PanoramicCapture/Private/SceneCapturer.cpp
--- a/Engine/Plugins/Experimental/PanoramicCapture/Source/PanoramicCapture/Private/SceneCapturer.cpp
+++ b/Engine/Plugins/Experimental/PanoramicCapture/Source/PanoramicCapture/Private/SceneCapturer.cpp
@@ -472,8 +472,8 @@ void USceneCapturer::RemoveAllBlendables(USceneCaptureComponent2D* CaptureCompon
 // Disable screen space post processes we cannot use while capturing
 void USceneCapturer::DisableUnsupportedPostProcesses(USceneCaptureComponent2D* CaptureComponent)
 {
-	CaptureComponent->PostProcessSettings.bOverride_GrainIntensity = true;
-	CaptureComponent->PostProcessSettings.GrainIntensity = 0.0f;
+	CaptureComponent->PostProcessSettings.bOverride_FilmGrainIntensity = true;
+	CaptureComponent->PostProcessSettings.FilmGrainIntensity = 0.0f;
 	CaptureComponent->PostProcessSettings.bOverride_MotionBlurAmount = true;
 	CaptureComponent->PostProcessSettings.MotionBlurAmount = 0.0f;
 	CaptureComponent->PostProcessSettings.bOverride_ScreenSpaceReflectionIntensity = true;
