PanoramicCapture: Fix conflict with nDisplay since CDO constructor forces display settings unnecessarily.

diff --git a/Engine/Plugins/Experimental/PanoramicCapture/Source/PanoramicCapture/Private/SceneCapturer.cpp b/Engine/Plugins/Experimental/PanoramicCapture/Source/PanoramicCapture/Private/SceneCapturer.cpp
--- a/Engine/Plugins/Experimental/PanoramicCapture/Source/PanoramicCapture/Private/SceneCapturer.cpp
+++ b/Engine/Plugins/Experimental/PanoramicCapture/Source/PanoramicCapture/Private/SceneCapturer.cpp
@@ -170,6 +170,11 @@ USceneCapturer::USceneCapturer()
 	, bOutputAmbientOcclusion(FStereoPanoramaManager::OutputAmbientOcclusion->GetInt() != 0)
 	, bMonoscopicMode(FStereoPanoramaManager::MonoscopicMode->GetInt() != 0)
 {
+	if (HasAnyFlags(RF_ArchetypeObject | RF_ClassDefaultObject))
+	{
+		return;
+	}
+
 	//NOTE: ikrimae: Keeping the old sampling mechanism just until we're sure the new way is always better
 	dbgMatchCaptureSliceFovToAtlasSliceFov = false;
 
