PanoramicCapture: Fix e.g. C: not recognized as a valid path.

#8981

diff --git a/Engine/Plugins/Experimental/PanoramicCapture/Source/PanoramicCapture/Private/SceneCapturer.cpp b/Engine/Plugins/Experimental/PanoramicCapture/Source/PanoramicCapture/Private/SceneCapturer.cpp
--- a/Engine/Plugins/Experimental/PanoramicCapture/Source/PanoramicCapture/Private/SceneCapturer.cpp
+++ b/Engine/Plugins/Experimental/PanoramicCapture/Source/PanoramicCapture/Private/SceneCapturer.cpp
@@ -19,6 +19,7 @@
 #include "GameFramework/GameModeBase.h"
 #include "GameFramework/Pawn.h"
 #include "GameFramework/PlayerController.h"
+#include "GenericPlatform/GenericPlatformMisc.h"
 #include "TextureResource.h"
 #include "Engine/BlendableInterface.h"
 #include "ImageUtils.h"
@@ -646,7 +647,7 @@ void USceneCapturer::ValidateParameters()
 	FString Drive = OutputDir.Left(Index + 1);
 	if (!Drive.IsEmpty() && FPaths::IsDrive(Drive))
 	{
-		if (!FPaths::DirectoryExists(Drive))
+		if (!FPaths::DirectoryExists(Drive + FGenericPlatformMisc::GetDefaultPathSeparator()))
 		{
 			ErrorFound = true;
 			FMessageLog(StereoPanoramaLogName).Message(EMessageSeverity::Error, LOCTEXT("ValidationError_MissingOutputDirectory", "The output directory's drive doesn't exists. Plese set SP.OutputDir with a valid path. Skipping renders..."));
